ElectraCodec: Renamed AV1 decoder factory

#rnx

diff --git a/Engine/Plugins/Media/AV1DecoderElectra/Source/AV1DecoderElectra/Private/AV1Decoder/ElectraMediaAV1Decoder.cpp b/Engine/Plugins/Media/AV1DecoderElectra/Source/AV1DecoderElectra/Private/AV1Decoder/ElectraMediaAV1Decoder.cpp
--- a/Engine/Plugins/Media/AV1DecoderElectra/Source/AV1DecoderElectra/Private/AV1Decoder/ElectraMediaAV1Decoder.cpp
+++ b/Engine/Plugins/Media/AV1DecoderElectra/Source/AV1DecoderElectra/Private/AV1Decoder/ElectraMediaAV1Decoder.cpp
@@ -13,7 +13,6 @@
 #include "IElectraDecoder.h"
 #include "IElectraDecoderOutputVideo.h"
 #include "ElectraDecodersUtils.h"
-//#include "Utils/Google/ElectraUtilsVPxVideo.h"
 #include "Utils/AOMedia/ElectraBitstreamProcessor_AV1.h"
 #include "Containers/Queue.h"
 #include <Stats/Stats.h>
@@ -307,10 +306,10 @@ private:
 /*********************************************************************************************************************/
 /*********************************************************************************************************************/
 
-class FAOMVideoDecoderElectraFactory : public IElectraCodecFactory, public IElectraCodecModularFeature, public TSharedFromThis<FAOMVideoDecoderElectraFactory, ESPMode::ThreadSafe>
+class FAV1VideoDecoderElectraFactory : public IElectraCodecFactory, public IElectraCodecModularFeature, public TSharedFromThis<FAV1VideoDecoderElectraFactory, ESPMode::ThreadSafe>
 {
 public:
-	virtual ~FAOMVideoDecoderElectraFactory()
+	virtual ~FAV1VideoDecoderElectraFactory()
 	{ }
 	void GetListOfFactories(TArray<TWeakPtr<IElectraCodecFactory, ESPMode::ThreadSafe>>& OutCodecFactories) override
 	{
@@ -344,9 +343,9 @@ public:
 		return MakeShared<FVideoDecoderAV1Electra, ESPMode::ThreadSafe>(InOptions);
 	}
 
-	static TSharedPtr<FAOMVideoDecoderElectraFactory, ESPMode::ThreadSafe> Self;
+	static TSharedPtr<FAV1VideoDecoderElectraFactory, ESPMode::ThreadSafe> Self;
 };
-TSharedPtr<FAOMVideoDecoderElectraFactory, ESPMode::ThreadSafe> FAOMVideoDecoderElectraFactory::Self;
+TSharedPtr<FAV1VideoDecoderElectraFactory, ESPMode::ThreadSafe> FAV1VideoDecoderElectraFactory::Self;
 
 #endif // WITH_LIBDAV1D_AV1_DECODER
 
@@ -361,18 +360,18 @@ void FElectraMediaAV1Decoder::Startup()
 	FModuleManager::Get().LoadModule(TEXT("ElectraCodecFactory"));
 
 	// Create an instance of the factory, which is also the modular feature.
-	check(!FAOMVideoDecoderElectraFactory::Self.IsValid());
-	FAOMVideoDecoderElectraFactory::Self = MakeShared<FAOMVideoDecoderElectraFactory, ESPMode::ThreadSafe>();
+	check(!FAV1VideoDecoderElectraFactory::Self.IsValid());
+	FAV1VideoDecoderElectraFactory::Self = MakeShared<FAV1VideoDecoderElectraFactory, ESPMode::ThreadSafe>();
 	// Register as modular feature.
-	IModularFeatures::Get().RegisterModularFeature(IElectraCodecFactoryModule::GetModularFeatureName(), FAOMVideoDecoderElectraFactory::Self.Get());
+	IModularFeatures::Get().RegisterModularFeature(IElectraCodecFactoryModule::GetModularFeatureName(), FAV1VideoDecoderElectraFactory::Self.Get());
 #endif
 }
 
 void FElectraMediaAV1Decoder::Shutdown()
 {
 #if WITH_LIBDAV1D_AV1_DECODER
-	IModularFeatures::Get().UnregisterModularFeature(IElectraCodecFactoryModule::GetModularFeatureName(), FAOMVideoDecoderElectraFactory::Self.Get());
-	FAOMVideoDecoderElectraFactory::Self.Reset();
+	IModularFeatures::Get().UnregisterModularFeature(IElectraCodecFactoryModule::GetModularFeatureName(), FAV1VideoDecoderElectraFactory::Self.Get());
+	FAV1VideoDecoderElectraFactory::Self.Reset();
 #endif
 }
 
